enabling sequencer scripting for mfg template

diff --git a/Templates/TP_MFG_BlankBP/TP_MFG_BlankBP.uproject b/Templates/TP_MFG_BlankBP/TP_MFG_BlankBP.uproject
--- a/Templates/TP_MFG_BlankBP/TP_MFG_BlankBP.uproject
+++ b/Templates/TP_MFG_BlankBP/TP_MFG_BlankBP.uproject
@@ -31,7 +31,7 @@
 		},
 		{
 			"Name": "SequencerScripting",
-			"Enabled": false
+			"Enabled": true
 		}
 	]
 }
\ No newline at end of file
